UE4 - change errors to warnings to avoid stopping automated tests

diff --git a/Engine/Plugins/Runtime/WindowsMoviePlayer/Source/WindowsMoviePlayer/Private/WindowsMoviePlayer.cpp b/Engine/Plugins/Runtime/WindowsMoviePlayer/Source/WindowsMoviePlayer/Private/WindowsMoviePlayer.cpp
--- a/Engine/Plugins/Runtime/WindowsMoviePlayer/Source/WindowsMoviePlayer/Private/WindowsMoviePlayer.cpp
+++ b/Engine/Plugins/Runtime/WindowsMoviePlayer/Source/WindowsMoviePlayer/Private/WindowsMoviePlayer.cpp
@@ -20,22 +20,22 @@ class FWindowsMoviePlayerModule : public IModuleInterface
 		// now attempt to load the delay loaded DLLs
 		if (LoadLibraryW(TEXT("shlwapi.dll")) == NULL)
 		{
-			UE_LOG( LogWindowsMoviePlayer, Error, TEXT("Could not load shlwapi.dll") );
+			UE_LOG(LogWindowsMoviePlayer, Warning, TEXT("Could not load shlwapi.dll") );
 			bLoadSuccessful = false;
 		}
 		if (LoadLibraryW(TEXT("mf.dll")) == NULL)
 		{
-			UE_LOG( LogWindowsMoviePlayer, Error, TEXT("Could not load mf.dll") );
+			UE_LOG(LogWindowsMoviePlayer, Warning, TEXT("Could not load mf.dll"));
 			bLoadSuccessful = false;
 		}
 		if (LoadLibraryW(TEXT("mfplat.dll")) == NULL)
 		{
-			UE_LOG( LogWindowsMoviePlayer, Error, TEXT("Could not load mfplat.dll") );
+			UE_LOG(LogWindowsMoviePlayer, Warning, TEXT("Could not load mfplat.dll"));
 			bLoadSuccessful = false;
 		}
 		if (LoadLibraryW(TEXT("mfplay.dll")) == NULL)
 		{
-			UE_LOG( LogWindowsMoviePlayer, Error, TEXT("Could not load mfplay.dll") );
+			UE_LOG(LogWindowsMoviePlayer, Warning, TEXT("Could not load mfplay.dll"));
 			bLoadSuccessful = false;
 		}
 
diff --git a/Engine/Source/Runtime/Core/Private/Windows/WindowsTextInputMethodSystem.cpp b/Engine/Source/Runtime/Core/Private/Windows/WindowsTextInputMethodSystem.cpp
--- a/Engine/Source/Runtime/Core/Private/Windows/WindowsTextInputMethodSystem.cpp
+++ b/Engine/Source/Runtime/Core/Private/Windows/WindowsTextInputMethodSystem.cpp
@@ -292,7 +292,7 @@ bool FWindowsTextInputMethodSystem::InitializeTSF()
 		Result = ::CoCreateInstance(CLSID_TF_ThreadMgr, NULL, CLSCTX_INPROC_SERVER, IID_ITfThreadMgr, reinterpret_cast<void**>(&(RawPointerTSFThreadManager)));
 		if(FAILED(Result))
 		{
-			UE_LOG(LogWindowsTextInputMethodSystem, Error, TEXT("Initialzation failed while creating the TSF thread manager."));
+			UE_LOG(LogWindowsTextInputMethodSystem, Warning, TEXT("Initialzation failed while creating the TSF thread manager."));
 			TSFInputProcessorProfiles.Reset();
 			TSFInputProcessorProfileManager.Reset();
 			return false;
@@ -368,7 +368,7 @@ bool FWindowsTextInputMethodSystem::InitializeTSF()
 	Result = TSFThreadManager->CreateDocumentMgr(&(TSFDisabledDocumentManager));
 	if(FAILED(Result))
 	{
-		UE_LOG(LogWindowsTextInputMethodSystem, Error, TEXT("Initialzation failed while creating the TSF thread manager."));
+		UE_LOG(LogWindowsTextInputMethodSystem, Warning, TEXT("Initialzation failed while creating the TSF thread manager."));
 		TSFInputProcessorProfiles.Reset();
 		TSFInputProcessorProfileManager.Reset();
 		TSFThreadManager.Reset();
