Miscellaneous syntax fixes when compiling on Windows using the Clang compiler

diff --git a/Engine/Plugins/Runtime/WindowsMoviePlayer/Source/WindowsMoviePlayer/Private/WindowsMovieStreamer.cpp b/Engine/Plugins/Runtime/WindowsMoviePlayer/Source/WindowsMoviePlayer/Private/WindowsMovieStreamer.cpp
--- a/Engine/Plugins/Runtime/WindowsMoviePlayer/Source/WindowsMoviePlayer/Private/WindowsMovieStreamer.cpp
+++ b/Engine/Plugins/Runtime/WindowsMoviePlayer/Source/WindowsMoviePlayer/Private/WindowsMovieStreamer.cpp
@@ -25,9 +25,9 @@
 DEFINE_LOG_CATEGORY(LogWindowsMoviePlayer);
 
 FMediaFoundationMovieStreamer::FMediaFoundationMovieStreamer()
-	: VideoPlayer(NULL)
+	: TextureData()
+	, VideoPlayer(NULL)
 	, SampleGrabberCallback(NULL)
-	, TextureData()
 {
 	MovieViewport = MakeShareable(new FMovieViewport());
 }
diff --git a/Engine/Source/Runtime/Engine/Private/GameInstance.cpp b/Engine/Source/Runtime/Engine/Private/GameInstance.cpp
--- a/Engine/Source/Runtime/Engine/Private/GameInstance.cpp
+++ b/Engine/Source/Runtime/Engine/Private/GameInstance.cpp
@@ -198,7 +198,7 @@ bool UGameInstance::StartPIEGameInstance(ULocalPlayer* LocalPlayer, bool bInSimu
 			if (!LocalPlayer->SpawnPlayActor(URL.ToString(1), Error, PlayWorld))
 			{
 				FMessageDialog::Open(EAppMsgType::Ok, FText::Format(NSLOCTEXT("UnrealEd", "Error_CouldntSpawnPlayer", "Couldn't spawn player: {0}"), FText::FromString(Error)));
-				return nullptr;
+				return false;
 			}
 		}
 
diff --git a/Engine/Source/Runtime/SlateCore/Private/Rendering/DrawElements.cpp b/Engine/Source/Runtime/SlateCore/Private/Rendering/DrawElements.cpp
--- a/Engine/Source/Runtime/SlateCore/Private/Rendering/DrawElements.cpp
+++ b/Engine/Source/Runtime/SlateCore/Private/Rendering/DrawElements.cpp
@@ -11,7 +11,7 @@ void FSlateDrawElement::Init(uint32 InLayer, const FPaintGeometry& PaintGeometry
 	Layer = InLayer;
 	Scale = PaintGeometry.DrawScale;
 	DrawEffects = InDrawEffects;
-	extern TOptional<FShortRect> GSlateScissorRect;
+	extern SLATECORE_API TOptional<FShortRect> GSlateScissorRect;
 	ScissorRect = GSlateScissorRect;
 }
 
