VS2015 - WMF changes, VC14 does not like QITABENT used in this fashion.

diff --git a/Engine/Plugins/Media/WmfMedia/Source/WmfMedia/Private/Wmf/WmfMediaByteStream.cpp b/Engine/Plugins/Media/WmfMedia/Source/WmfMedia/Private/Wmf/WmfMediaByteStream.cpp
--- a/Engine/Plugins/Media/WmfMedia/Source/WmfMedia/Private/Wmf/WmfMediaByteStream.cpp
+++ b/Engine/Plugins/Media/WmfMedia/Source/WmfMedia/Private/Wmf/WmfMediaByteStream.cpp
@@ -81,6 +81,10 @@ STDMETHODIMP FWmfMediaByteStream::Invoke(IMFAsyncResult* AsyncResult)
 	return S_OK;
 }
 
+#if _MSC_VER == 1900
+#pragma warning(push)
+#pragma warning(disable:4838)
+#endif // _MSC_VER == 1900
 
 STDMETHODIMP FWmfMediaByteStream::QueryInterface(REFIID RefID, void** Object)
 {
@@ -93,6 +97,9 @@ STDMETHODIMP FWmfMediaByteStream::QueryInterface(REFIID RefID, void** Object)
 
 	return QISearch( this, QITab, RefID, Object );
 }
+#if _MSC_VER == 1900
+#pragma warning(pop)
+#endif // _MSC_VER == 1900
 
 
 STDMETHODIMP_(ULONG) FWmfMediaByteStream::Release()
diff --git a/Engine/Plugins/Media/WmfMedia/Source/WmfMedia/Private/Wmf/WmfMediaReadState.h b/Engine/Plugins/Media/WmfMedia/Source/WmfMedia/Private/Wmf/WmfMediaReadState.h
--- a/Engine/Plugins/Media/WmfMedia/Source/WmfMedia/Private/Wmf/WmfMediaReadState.h
+++ b/Engine/Plugins/Media/WmfMedia/Source/WmfMedia/Private/Wmf/WmfMediaReadState.h
@@ -79,6 +79,11 @@ public:
 
 	// IUnknown interface
 
+
+#if _MSC_VER == 1900
+#pragma warning(push)
+#pragma warning(disable:4838)
+#endif // _MSC_VER == 1900
 	STDMETHODIMP QueryInterface( REFIID RefID, void** Object )
 	{
 		static const QITAB QITab[] =
@@ -89,6 +94,9 @@ public:
 
 		return QISearch( this, QITab, RefID, Object );
 	}
+#if _MSC_VER == 1900
+#pragma warning(pop)
+#endif // _MSC_VER == 1900
 
 	STDMETHODIMP_(ULONG) AddRef()
 	{
diff --git a/Engine/Plugins/Media/WmfMedia/Source/WmfMedia/Private/Wmf/WmfMediaSampler.h b/Engine/Plugins/Media/WmfMedia/Source/WmfMedia/Private/Wmf/WmfMediaSampler.h
--- a/Engine/Plugins/Media/WmfMedia/Source/WmfMedia/Private/Wmf/WmfMediaSampler.h
+++ b/Engine/Plugins/Media/WmfMedia/Source/WmfMedia/Private/Wmf/WmfMediaSampler.h
@@ -119,6 +119,10 @@ public:
 	}
 
 
+#if _MSC_VER == 1900
+#pragma warning(push)
+#pragma warning(disable:4838)
+#endif // _MSC_VER == 1900
 	STDMETHODIMP QueryInterface( REFIID RefID, void** Object )
 	{
 		static const QITAB QITab[] =
@@ -129,6 +133,9 @@ public:
 
 		return QISearch(this, QITab, RefID, Object);
 	}
+#if _MSC_VER == 1900
+#pragma warning(pop)
+#endif // _MSC_VER == 1900
 
 	STDMETHODIMP_(ULONG) Release()
 	{
diff --git a/Engine/Plugins/Media/WmfMedia/Source/WmfMedia/Private/Wmf/WmfMediaSession.cpp b/Engine/Plugins/Media/WmfMedia/Source/WmfMedia/Private/Wmf/WmfMediaSession.cpp
--- a/Engine/Plugins/Media/WmfMedia/Source/WmfMedia/Private/Wmf/WmfMediaSession.cpp
+++ b/Engine/Plugins/Media/WmfMedia/Source/WmfMedia/Private/Wmf/WmfMediaSession.cpp
@@ -294,6 +294,11 @@ STDMETHODIMP FWmfMediaSession::Invoke( IMFAsyncResult* AsyncResult )
 }
 
 
+#if _MSC_VER == 1900
+#pragma warning(push)
+#pragma warning(disable:4838)
+#endif // _MSC_VER == 1900
+
 STDMETHODIMP FWmfMediaSession::QueryInterface( REFIID RefID, void** Object )
 {
 	static const QITAB QITab[] =
@@ -304,6 +309,9 @@ STDMETHODIMP FWmfMediaSession::QueryInterface( REFIID RefID, void** Object )
 
 	return QISearch(this, QITab, RefID, Object);
 }
+#if _MSC_VER == 1900
+#pragma warning(pop)
+#endif // _MSC_VER == 1900
 
 
 STDMETHODIMP_(ULONG) FWmfMediaSession::Release()
diff --git a/Engine/Plugins/Runtime/WindowsMoviePlayer/Source/WindowsMoviePlayer/Private/WindowsMovieStreamer.cpp b/Engine/Plugins/Runtime/WindowsMoviePlayer/Source/WindowsMoviePlayer/Private/WindowsMovieStreamer.cpp
--- a/Engine/Plugins/Runtime/WindowsMoviePlayer/Source/WindowsMoviePlayer/Private/WindowsMovieStreamer.cpp
+++ b/Engine/Plugins/Runtime/WindowsMoviePlayer/Source/WindowsMoviePlayer/Private/WindowsMovieStreamer.cpp
@@ -180,6 +180,10 @@ void FMediaFoundationMovieStreamer::CleanupRenderingResources()
 	}
 }
 
+#if _MSC_VER == 1900
+#pragma warning(push)
+#pragma warning(disable:4838)
+#endif // _MSC_VER == 1900
 STDMETHODIMP FVideoPlayer::QueryInterface(REFIID RefID, void** Object)
 {
 	static const QITAB QITab[] =
@@ -189,6 +193,9 @@ STDMETHODIMP FVideoPlayer::QueryInterface(REFIID RefID, void** Object)
 	};
 	return QISearch(this, QITab, RefID, Object);
 }
+#if _MSC_VER == 1900
+#pragma warning(pop)
+#endif // _MSC_VER == 1900
 
 STDMETHODIMP_(ULONG) FVideoPlayer::AddRef()
 {
@@ -475,6 +482,10 @@ FIntPoint FVideoPlayer::AddStreamToTopology(IMFTopology* Topology, IMFPresentati
 
 
 
+#if _MSC_VER == 1900
+#pragma warning(push)
+#pragma warning(disable:4838)
+#endif // _MSC_VER == 1900
 STDMETHODIMP FSampleGrabberCallback::QueryInterface(REFIID RefID, void** Object)
 {
 	static const QITAB QITab[] =
@@ -485,6 +496,9 @@ STDMETHODIMP FSampleGrabberCallback::QueryInterface(REFIID RefID, void** Object)
 	};
 	return QISearch(this, QITab, RefID, Object);
 }
+#if _MSC_VER == 1900
+#pragma warning(pop)
+#endif // _MSC_VER == 1900
 
 STDMETHODIMP_(ULONG) FSampleGrabberCallback::AddRef()
 {
