Fix white frames on PC and XB1 startup movie players

diff --git a/Engine/Plugins/Runtime/WindowsMoviePlayer/Source/WindowsMoviePlayer/Private/WindowsMovieStreamer.cpp b/Engine/Plugins/Runtime/WindowsMoviePlayer/Source/WindowsMoviePlayer/Private/WindowsMovieStreamer.cpp
--- a/Engine/Plugins/Runtime/WindowsMoviePlayer/Source/WindowsMoviePlayer/Private/WindowsMovieStreamer.cpp
+++ b/Engine/Plugins/Runtime/WindowsMoviePlayer/Source/WindowsMoviePlayer/Private/WindowsMovieStreamer.cpp
@@ -39,6 +39,8 @@ FMediaFoundationMovieStreamer::~FMediaFoundationMovieStreamer()
 	CloseMovie();
 	CleanupRenderingResources();
 
+	FlushRenderingCommands();
+
 	TextureFreeList.Empty();
 }
 
@@ -53,6 +55,8 @@ bool FMediaFoundationMovieStreamer::Init(const TArray<FString>& MoviePaths, TEnu
 	PlaybackType = inPlaybackType;
 	StoredMoviePaths = MoviePaths;
 
+	MovieViewport->SetTexture(nullptr);
+
 	OpenNextMovie();
 
 	return true;
@@ -80,6 +84,11 @@ bool FMediaFoundationMovieStreamer::Tick(float DeltaTime)
 		FMemory::Memcpy( DestTextureData, TextureData.GetData(), TextureData.Num() );
 		RHIUnlockTexture2D( CurrentTexture->GetTypedResource(), 0, false );
 
+		if (MovieViewport->GetViewportRenderTargetTexture() == nullptr)
+		{
+			MovieViewport->SetTexture(Texture);
+		}
+
 		SampleGrabberCallback->SetNeedNewSample();
 	}
 
@@ -89,12 +98,18 @@ bool FMediaFoundationMovieStreamer::Tick(float DeltaTime)
 		if ((MovieIndex + 1) < StoredMoviePaths.Num())
 		{
 			++MovieIndex;
-			OpenNextMovie();
+			if (OpenNextMovie())
+			{
+				MovieViewport->SetTexture(Texture);
+			}
 		}
 		else if (PlaybackType != MT_Normal)
 		{
 			MovieIndex = PlaybackType == MT_LoadingLoop ? StoredMoviePaths.Num() - 1 : 0;
-			OpenNextMovie();
+			if (OpenNextMovie())
+			{
+				MovieViewport->SetTexture(Texture);
+			}
 		}
 		else
 		{
@@ -121,7 +136,7 @@ void FMediaFoundationMovieStreamer::Cleanup()
 	CleanupRenderingResources();
 }
 
-void FMediaFoundationMovieStreamer::OpenNextMovie()
+bool FMediaFoundationMovieStreamer::OpenNextMovie()
 {
 	check(StoredMoviePaths.Num() > 0 && MovieIndex < StoredMoviePaths.Num());
 	FString MoviePath = FPaths::ProjectContentDir() + TEXT("Movies/") + StoredMoviePaths[MovieIndex];
@@ -162,8 +177,12 @@ void FMediaFoundationMovieStreamer::OpenNextMovie()
 			});
 		}
 
-		MovieViewport->SetTexture(Texture);
 		VideoPlayer->StartPlayback();
+		return true;
+	}
+	else
+	{
+		return false;
 	}
 }
 
diff --git a/Engine/Plugins/Runtime/WindowsMoviePlayer/Source/WindowsMoviePlayer/Private/WindowsMovieStreamer.h b/Engine/Plugins/Runtime/WindowsMoviePlayer/Source/WindowsMoviePlayer/Private/WindowsMovieStreamer.h
--- a/Engine/Plugins/Runtime/WindowsMoviePlayer/Source/WindowsMoviePlayer/Private/WindowsMovieStreamer.h
+++ b/Engine/Plugins/Runtime/WindowsMoviePlayer/Source/WindowsMoviePlayer/Private/WindowsMovieStreamer.h
@@ -47,7 +47,7 @@ public:
 
 private:
 	/** Opens up the next movie in the movie path queue */
-	void OpenNextMovie();
+	bool OpenNextMovie();
 	/** Closes the currently running video */
 	void CloseMovie();
 	/** Cleans up rendering resources once movies are done playing */
