Startup movies on Windows platform can now be 1920x1080.
Previously the resolution needed to be a multiple of 16 (such as 1920x1088) due to using MFVideoFormat_RGB32 setting.

All windows startup movies are now processed as YUY2 and drawn using conversion shader.

diff --git a/Engine/Plugins/Runtime/WindowsMoviePlayer/Source/WindowsMoviePlayer/Private/WindowsMovieStreamer.cpp b/Engine/Plugins/Runtime/WindowsMoviePlayer/Source/WindowsMoviePlayer/Private/WindowsMovieStreamer.cpp
--- a/Engine/Plugins/Runtime/WindowsMoviePlayer/Source/WindowsMoviePlayer/Private/WindowsMovieStreamer.cpp
+++ b/Engine/Plugins/Runtime/WindowsMoviePlayer/Source/WindowsMoviePlayer/Private/WindowsMovieStreamer.cpp
@@ -8,6 +8,10 @@
 #include "RenderUtils.h"
 #include "Runtime/HeadMountedDisplay/Public/IHeadMountedDisplayModule.h"
 
+#include "MediaShaders.h"
+#include "PipelineStateCache.h"
+#include "RHIStaticStates.h"
+
 #pragma comment(lib, "shlwapi")
 #pragma comment(lib, "mf")
 #pragma comment(lib, "mfplat")
@@ -67,6 +71,102 @@ void FMediaFoundationMovieStreamer::ForceCompletion()
 	CloseMovie();
 }
 
+void FMediaFoundationMovieStreamer::ConvertSample()
+{
+	const bool SrgbOutput = false;
+	const bool bSampleIsOutputSrgb = false;
+
+	const FMovieTrackFormat& SourceFormat = VideoPlayer->GetVideoTrackFormat();
+	const EPixelFormat InputPixelFormat = PF_B8G8R8A8;
+
+	{
+		const bool SrgbTexture = false;
+		const uint32 InputCreateFlags = TexCreate_Dynamic | (SrgbTexture ? TexCreate_SRGB : 0);
+
+		// create a new input render target if necessary
+		if (!InputTarget.IsValid() || (InputTarget->GetSizeXY() != SourceFormat.BufferDim) || (InputTarget->GetFormat() != InputPixelFormat) || ((InputTarget->GetFlags() & InputCreateFlags) != InputCreateFlags))
+		{
+			TRefCountPtr<FRHITexture2D> DummyTexture2DRHI;
+			FRHIResourceCreateInfo CreateInfo;
+
+			RHICreateTargetableShaderResource2D(
+				SourceFormat.BufferDim.X,
+				SourceFormat.BufferDim.Y,
+				InputPixelFormat,
+				1,
+				InputCreateFlags,
+				TexCreate_RenderTargetable,
+				false,
+				CreateInfo,
+				InputTarget,
+				DummyTexture2DRHI
+			);
+		}
+
+		// copy sample data to input render target
+		FUpdateTextureRegion2D Region(0, 0, 0, 0, SourceFormat.BufferDim.X, SourceFormat.BufferDim.Y);
+		RHIUpdateTexture2D(InputTarget, 0, Region, SourceFormat.BufferStride, TextureData.GetData());
+	}
+
+	const FIntPoint OutputDim = SourceFormat.OutputDim;
+
+	FSlateTexture2DRHIRef* CurrentTexture = Texture.Get();
+	FTextureRHIParamRef RenderTarget = CurrentTexture->GetRHIRef();
+
+	// perform the conversion
+	FRHICommandListImmediate& CommandList = FRHICommandListExecutor::GetImmediateCommandList();
+	
+	FGraphicsPipelineStateInitializer GraphicsPSOInit;
+	SetRenderTargets(CommandList, 1, &RenderTarget, nullptr, ESimpleRenderTargetMode::EExistingColorAndDepth, FExclusiveDepthStencil::DepthNop_StencilNop);
+
+	CommandList.ApplyCachedRenderTargets(GraphicsPSOInit);
+	CommandList.SetViewport(0, 0, 0.0f, OutputDim.X, OutputDim.Y, 1.0f);
+
+	GraphicsPSOInit.DepthStencilState = TStaticDepthStencilState<false, CF_Always>::GetRHI();
+	GraphicsPSOInit.RasterizerState = TStaticRasterizerState<>::GetRHI();
+	GraphicsPSOInit.BlendState = TStaticBlendStateWriteMask<CW_RGBA, CW_NONE, CW_NONE, CW_NONE, CW_NONE, CW_NONE, CW_NONE, CW_NONE>::GetRHI();
+	GraphicsPSOInit.PrimitiveType = PT_TriangleStrip;
+
+	// configure media shaders
+	auto ShaderMap = GetGlobalShaderMap(GMaxRHIFeatureLevel);
+	TShaderMapRef<FMediaShadersVS> VertexShader(ShaderMap);
+
+	GraphicsPSOInit.BoundShaderState.VertexDeclarationRHI = GMediaVertexDeclaration.VertexDeclarationRHI;
+	GraphicsPSOInit.BoundShaderState.VertexShaderRHI = GETSAFERHISHADER_VERTEX(*VertexShader);
+
+	switch (SourceFormat.SampleFormat)
+	{
+	case EMediaTextureSampleFormat::CharBMP:
+	{
+		TShaderMapRef<FBMPConvertPS> ConvertShader(ShaderMap);
+		GraphicsPSOInit.BoundShaderState.PixelShaderRHI = GETSAFERHISHADER_PIXEL(*ConvertShader);
+		SetGraphicsPipelineState(CommandList, GraphicsPSOInit);
+		ConvertShader->SetParameters(CommandList, InputTarget, OutputDim, bSampleIsOutputSrgb && !SrgbOutput);
+	}
+	break;
+	
+	case EMediaTextureSampleFormat::CharYUY2:
+	{
+		TShaderMapRef<FYUY2ConvertPS> ConvertShader(ShaderMap);
+		GraphicsPSOInit.BoundShaderState.PixelShaderRHI = GETSAFERHISHADER_PIXEL(*ConvertShader);
+		SetGraphicsPipelineState(CommandList, GraphicsPSOInit);
+		ConvertShader->SetParameters(CommandList, InputTarget, OutputDim, MediaShaders::YuvToSrgbDefault, bSampleIsOutputSrgb);
+	}
+	break;
+
+	default:
+		return; // unsupported format
+	}
+
+	// draw full size quad into render target
+	FVertexBufferRHIRef VertexBuffer = CreateTempMediaVertexBuffer();
+	CommandList.SetStreamSource(0, VertexBuffer, 0);
+	// set viewport to RT size
+	CommandList.SetViewport(0, 0, 0.0f, OutputDim.X, OutputDim.Y, 1.0f);
+	CommandList.DrawPrimitive(PT_TriangleStrip, 0, 2, 1);
+	CommandList.TransitionResource(EResourceTransitionAccess::EReadable, RenderTarget);
+}
+
 bool FMediaFoundationMovieStreamer::Tick(float DeltaTime)
 {
 	FSlateTexture2DRHIRef* CurrentTexture = Texture.Get();
@@ -79,10 +179,18 @@ bool FMediaFoundationMovieStreamer::Tick(float DeltaTime)
 
 	if (CurrentTexture && SampleGrabberCallback->GetIsSampleReadyToUpdate())
 	{
-		uint32 Stride;
-		uint8* DestTextureData = (uint8*)RHILockTexture2D( CurrentTexture->GetTypedResource(), 0, RLM_WriteOnly, Stride, false );
-		FMemory::Memcpy( DestTextureData, TextureData.GetData(), TextureData.Num() );
-		RHIUnlockTexture2D( CurrentTexture->GetTypedResource(), 0, false );
+		const FMovieTrackFormat& SourceFormat = VideoPlayer->GetVideoTrackFormat();
+		if (SourceFormat.SampleFormat == EMediaTextureSampleFormat::CharBGRA && SourceFormat.BufferDim == SourceFormat.OutputDim)
+		{
+			uint32 Stride;
+			uint8* DestTextureData = (uint8*)RHILockTexture2D( CurrentTexture->GetTypedResource(), 0, RLM_WriteOnly, Stride, false );
+			FMemory::Memcpy( DestTextureData, TextureData.GetData(), TextureData.Num());
+			RHIUnlockTexture2D( CurrentTexture->GetTypedResource(), 0, false );
+		}
+		else
+		{
+			ConvertSample();
+		}
 
 		if (MovieViewport->GetViewportRenderTargetTexture() == nullptr)
 		{
@@ -149,7 +257,6 @@ bool FMediaFoundationMovieStreamer::OpenNextMovie()
 	if( VideoDimensions != FIntPoint::ZeroValue )
 	{
 		TextureData.Empty();
-		TextureData.AddZeroed(VideoDimensions.X*VideoDimensions.Y*GPixelFormats[PF_B8G8R8A8].BlockBytes);
 
 		if( TextureFreeList.Num() > 0 )
 		{
@@ -168,7 +275,7 @@ bool FMediaFoundationMovieStreamer::OpenNextMovie()
 		else
 		{
 			const bool bCreateEmptyTexture = true;
-			Texture = MakeShareable(new FSlateTexture2DRHIRef(VideoDimensions.X, VideoDimensions.Y, PF_B8G8R8A8, NULL, TexCreate_Dynamic, bCreateEmptyTexture));
+			Texture = MakeShareable(new FSlateTexture2DRHIRef(VideoDimensions.X, VideoDimensions.Y, PF_B8G8R8A8, NULL, TexCreate_RenderTargetable, bCreateEmptyTexture));
 
 			ENQUEUE_UNIQUE_RENDER_COMMAND_ONEPARAMETER(InitMovieTexture,
 				FSlateTexture2DRHIRef*, TextureRHIRef, Texture.Get(),
@@ -217,6 +324,8 @@ void FMediaFoundationMovieStreamer::CleanupRenderingResources()
 	{
 		BeginReleaseResource( TextureFreeList[TextureIndex].Get() );
 	}
+
+	InputTarget.SafeRelease();
 }
 
 #if _MSC_VER == 1900
@@ -474,10 +583,58 @@ FIntPoint FVideoPlayer::AddStreamToTopology(IMFTopology* Topology, IMFPresentati
 			HResult = MFGetAttributeSize(OutputType, MF_MT_FRAME_SIZE, &Width, &Height);
 			check(SUCCEEDED(HResult));
 
-			HResult = InputType->SetGUID(MF_MT_MAJOR_TYPE, MFMediaType_Video);
+			GUID SourceVideoSubType;
+			HResult = OutputType->GetGUID(MF_MT_SUBTYPE, &SourceVideoSubType);
 			check(SUCCEEDED(HResult));
-			HResult = InputType->SetGUID(MF_MT_SUBTYPE, MFVideoFormat_RGB32);
+
+			HResult = InputType->SetGUID(MF_MT_MAJOR_TYPE, MFMediaType_Video);
 			check(SUCCEEDED(HResult));
+
+			VideoTrackFormat.OutputDim = FIntPoint(Width, Height);
+			{
+				const bool Uncompressed =
+					(SourceVideoSubType == MFVideoFormat_RGB555) ||
+					(SourceVideoSubType == MFVideoFormat_RGB565) ||
+					(SourceVideoSubType == MFVideoFormat_RGB24) ||
+					(SourceVideoSubType == MFVideoFormat_RGB32) ||
+					(SourceVideoSubType == MFVideoFormat_ARGB32);
+
+				if (Uncompressed)
+				{
+					// Note: MFVideoFormat_RGB32 tends to require resolutions that are multiple of 16 preventing 1920x1080 from working
+					HResult = InputType->SetGUID(MF_MT_SUBTYPE, MFVideoFormat_RGB32);
+					check(SUCCEEDED(HResult));
+
+					VideoTrackFormat.SampleFormat = EMediaTextureSampleFormat::CharBMP;
+					VideoTrackFormat.BufferDim = VideoTrackFormat.OutputDim;
+					VideoTrackFormat.BufferStride = VideoTrackFormat.OutputDim.X * 4;
+				}
+				else
+				{
+					HResult = InputType->SetGUID(MF_MT_SUBTYPE, MFVideoFormat_YUY2);
+					check(SUCCEEDED(HResult));
+
+					VideoTrackFormat.SampleFormat = EMediaTextureSampleFormat::CharYUY2;
+					
+					int32 AlignedOutputX = VideoTrackFormat.OutputDim.X;
+					
+					if ((SourceVideoSubType == MFVideoFormat_H264) || (SourceVideoSubType == MFVideoFormat_H264_ES)) 
+					{
+						AlignedOutputX = Align(AlignedOutputX, 16);
+					}
+					
+					int32 SampleStride = AlignedOutputX * 2; // 2 bytes per pixel
+
+					if (SampleStride < 0)
+					{
+						SampleStride = -SampleStride;
+					}
+					
+					VideoTrackFormat.BufferDim = FIntPoint(AlignedOutputX / 2, VideoTrackFormat.OutputDim.Y); // 2 pixels per texel
+					VideoTrackFormat.BufferStride = SampleStride;
+				}
+			}
+
 			HResult = InputType->SetUINT32(MF_MT_ALL_SAMPLES_INDEPENDENT, TRUE);
 			check(SUCCEEDED(HResult));
 			HResult = MFCreateSampleGrabberSinkActivate(InputType, SampleGrabberCallback, &SinkActivate);
@@ -571,8 +728,11 @@ STDMETHODIMP FSampleGrabberCallback::OnProcessSample(REFGUID MajorMediaType, DWO
 {
 	if (VideoSampleReady.GetValue() == 0)
 	{
-		check(TextureData.Num() == SampleSize);
-		FMemory::Memcpy(&TextureData[0], SampleBuffer, SampleSize);
+		TextureData.SetNum(SampleSize, false);
+		if (SampleSize > 0)
+		{
+			FMemory::Memcpy(TextureData.GetData(), SampleBuffer, SampleSize);
+		}
 
 		VideoSampleReady.Set(1);
 	}
diff --git a/Engine/Plugins/Runtime/WindowsMoviePlayer/Source/WindowsMoviePlayer/Private/WindowsMovieStreamer.h b/Engine/Plugins/Runtime/WindowsMoviePlayer/Source/WindowsMoviePlayer/Private/WindowsMovieStreamer.h
--- a/Engine/Plugins/Runtime/WindowsMoviePlayer/Source/WindowsMoviePlayer/Private/WindowsMovieStreamer.h
+++ b/Engine/Plugins/Runtime/WindowsMoviePlayer/Source/WindowsMoviePlayer/Private/WindowsMovieStreamer.h
@@ -4,6 +4,7 @@
 
 #include "CoreMinimal.h"
 #include "MoviePlayer.h"
+#include "IMediaTextureSample.h"
 
 #include "Windows/WindowsHWrapper.h"
 #include "Windows/AllowWindowsPlatformTypes.h"
@@ -52,6 +53,8 @@ private:
 	void CloseMovie();
 	/** Cleans up rendering resources once movies are done playing */
 	void CleanupRenderingResources();
+	/** Use shader to generate final image */
+	void ConvertSample();
 
 private:
 	/** A list of all the stored movie paths we have enqueued for playing */
@@ -74,12 +77,31 @@ private:
 	 */
 	TArray<TSharedPtr<FSlateTexture2DRHIRef, ESPMode::ThreadSafe>> TextureFreeList;
 
+	/** Input texture for conversion shader */
+	TRefCountPtr<FRHITexture2D> InputTarget;
+
 	/** The video player and sample grabber for use of Media Foundation */
 	class FVideoPlayer* VideoPlayer;
     class FSampleGrabberCallback* SampleGrabberCallback;
 };
 
+/** Video track details */
+struct FMovieTrackFormat
+{
+public:
+	FMovieTrackFormat() :
+		BufferDim(FIntPoint(0, 0)),
+		BufferStride(0),
+		OutputDim(FIntPoint(0, 0)),
+		SampleFormat(EMediaTextureSampleFormat::Undefined)
+	{
+	}
 
+	FIntPoint BufferDim;
+	uint32 BufferStride;
+	FIntPoint OutputDim;
+	EMediaTextureSampleFormat SampleFormat;
+};
 
 /** The video player is the class which handles all the loading and playing of videos */
 class FVideoPlayer : public IMFAsyncCallback
@@ -116,6 +138,8 @@ public:
 	/** True if the movie is still playing and rendering frames */
 	bool MovieIsRunning() const {return MovieIsFinished.GetValue() == 0;}
 	
+	const FMovieTrackFormat& GetVideoTrackFormat() const { return VideoTrackFormat; }
+
 private:
 	/** Sets up the topology of all the nodes in the media session, returning the video dimensions */
 	FIntPoint SetPlaybackTopology(class FSampleGrabberCallback* SampleGrabberCallback);
@@ -135,6 +159,9 @@ private:
 	FThreadSafeCounter MovieIsFinished;
 	/** This counter locks the ticking thread while all Media Foundation threads shutdown */
 	FThreadSafeCounter CloseIsPosted;
+
+	/* Format of video track */
+	FMovieTrackFormat VideoTrackFormat;
 };
 
 
diff --git a/Engine/Plugins/Runtime/WindowsMoviePlayer/Source/WindowsMoviePlayer/WindowsMoviePlayer.Build.cs b/Engine/Plugins/Runtime/WindowsMoviePlayer/Source/WindowsMoviePlayer/WindowsMoviePlayer.Build.cs
--- a/Engine/Plugins/Runtime/WindowsMoviePlayer/Source/WindowsMoviePlayer/WindowsMoviePlayer.Build.cs
+++ b/Engine/Plugins/Runtime/WindowsMoviePlayer/Source/WindowsMoviePlayer/WindowsMoviePlayer.Build.cs
@@ -11,6 +11,11 @@ namespace UnrealBuildTool.Rules
 				}
 				);
 
+			PrivateIncludePathModuleNames.AddRange(
+				new string[] {
+					"Media"
+				});
+
 			PrivateIncludePaths.AddRange(
 				new string[] {
 				}
@@ -27,7 +32,9 @@ namespace UnrealBuildTool.Rules
                     "RenderCore",
                     "RHI",
                     "SlateCore",
-                    "Slate"
+                    "Slate",
+					"ShaderCore",
+					"UtilityShaders"
 				}
 				);
 
