Fix startup movies on Windows not playing when no audio devices enabled.
Also fixes startup movies hanging when audio devices disabled during playback.

diff --git a/Engine/Plugins/Runtime/WindowsMoviePlayer/Source/WindowsMoviePlayer/Private/WindowsMovieStreamer.cpp b/Engine/Plugins/Runtime/WindowsMoviePlayer/Source/WindowsMoviePlayer/Private/WindowsMovieStreamer.cpp
--- a/Engine/Plugins/Runtime/WindowsMoviePlayer/Source/WindowsMoviePlayer/Private/WindowsMovieStreamer.cpp
+++ b/Engine/Plugins/Runtime/WindowsMoviePlayer/Source/WindowsMoviePlayer/Private/WindowsMovieStreamer.cpp
@@ -19,12 +19,13 @@
 #pragma comment(lib, "mfuuid")
 
 #include "Windows/AllowWindowsPlatformTypes.h"
+#include "Windows/COMPointer.h"
 
 #include <windows.h>
 #include <shlwapi.h>
 #include <mfapi.h>
 #include <mfidl.h>
-
+#include <mferror.h>
 
 
 DEFINE_LOG_CATEGORY(LogWindowsMoviePlayer);
@@ -33,6 +34,7 @@ FMediaFoundationMovieStreamer::FMediaFoundationMovieStreamer()
 	: TextureData()
 	, VideoPlayer(NULL)
 	, SampleGrabberCallback(NULL)
+	, bUseSound(true)
 {
 	MovieViewport = MakeShareable(new FMovieViewport());
 	PlaybackType = MT_Normal;
@@ -202,6 +204,14 @@ bool FMediaFoundationMovieStreamer::Tick(float DeltaTime)
 
 	if (!VideoPlayer->MovieIsRunning())
 	{
+		// Playback can fail when no audio output devices enabled
+		if (VideoPlayer->FailedToCreateMediaSink() && bUseSound)
+		{
+			// Retry playback without audio
+			bUseSound = false;
+			--MovieIndex;
+		}
+
 		CloseMovie();
 		if ((MovieIndex + 1) < StoredMoviePaths.Num())
 		{
@@ -252,7 +262,7 @@ bool FMediaFoundationMovieStreamer::OpenNextMovie()
 	SampleGrabberCallback = new FSampleGrabberCallback(TextureData);
 	
 	VideoPlayer = new FVideoPlayer();
-	FIntPoint VideoDimensions = VideoPlayer->OpenFile(MoviePath, SampleGrabberCallback);
+	FIntPoint VideoDimensions = VideoPlayer->OpenFile(MoviePath, SampleGrabberCallback, bUseSound);
 
 	if( VideoDimensions != FIntPoint::ZeroValue )
 	{
@@ -359,54 +369,96 @@ STDMETHODIMP_(ULONG) FVideoPlayer::Release()
 
 HRESULT FVideoPlayer::Invoke(IMFAsyncResult* AsyncResult)
 {
-	IMFMediaEvent* Event = NULL;
+	TComPtr<IMFMediaEvent> Event;
 	HRESULT HResult = MediaSession->EndGetEvent(AsyncResult, &Event);
 	if (FAILED(HResult))
 	{
-		Event->Release();
 		return S_OK;
 	}
 
+	// get event type
 	MediaEventType EventType = MEUnknown;
 	HResult = Event->GetType(&EventType);
-	Event->Release();
 	if (FAILED(HResult))
 	{
 		return S_OK;
 	}
 
-	if (EventType == MESessionClosed)
+	// get event status
+	HRESULT EventStatus = S_FALSE;
+	HResult = Event->GetStatus(&EventStatus);
+	if (FAILED(HResult))
+	{
+		return S_OK;
+	}
+
+	bool bRequestNextEvent = true;
+	bool bFinishedAndClose = false;
+	switch (EventType)
+	{
+	case MEAudioSessionDeviceRemoved:
+		// Need to stop playback now or will be stuck forever
+		bFinishedAndClose = true;
+		break;
+
+	case MESessionClosed:
+		bFinishedAndClose = true;
+		break;
+	
+	case MESessionTopologySet:
+		if (!SUCCEEDED(EventStatus))
+		{
+			if (EventStatus == MF_E_CANNOT_CREATE_SINK)
+			{
+				bFailedToCreateMediaSink = true;
+			}
+
+			// Topology error
+			bFinishedAndClose = true;
+		}
+		break;
+
+	case MEEndOfPresentation:
+		if (MovieIsRunning())
+		{
+			MovieIsFinished.Set(1);
+		}
+		break;
+
+	case MEError:
+		if (MovieIsRunning())
+		{
+			// Unknown fatal error
+			bFinishedAndClose = true;
+		}
+		break;
+
+	default:
+		break;
+	}
+
+	if (bFinishedAndClose)
 	{
 		MovieIsFinished.Set(1);
 		CloseIsPosted.Set(1);
+		bRequestNextEvent = false;
 	}
-	else
+
+	if (bRequestNextEvent)
 	{
 		HResult = MediaSession->BeginGetEvent(this, NULL);
 		if (FAILED(HResult))
 		{
+			MovieIsFinished.Set(1);
+			CloseIsPosted.Set(1);
 			return S_OK;
 		}
-		
-		if (MovieIsRunning())
-		{
-			if (EventType == MEEndOfPresentation)
-			{
-				MovieIsFinished.Set(1);
-			}
-			else if( EventType == MEError )
-			{
-				// Unknown fatal error
-				MovieIsFinished.Set(1);
-				CloseIsPosted.Set(1);
-			}
-		}
 	}
 	
 	return S_OK;
 }
 
-FIntPoint FVideoPlayer::OpenFile(const FString& FilePath, class FSampleGrabberCallback* SampleGrabberCallback)
+FIntPoint FVideoPlayer::OpenFile(const FString& FilePath, class FSampleGrabberCallback* SampleGrabberCallback, bool bUseSound)
 {
 	FIntPoint OutDimensions = FIntPoint::ZeroValue;
 
@@ -437,7 +489,7 @@ FIntPoint FVideoPlayer::OpenFile(const FString& FilePath, class FSampleGrabberCa
 			HResult = Source->QueryInterface(IID_PPV_ARGS(&MediaSource));
 			Source->Release();
 			
-			OutDimensions = SetPlaybackTopology(SampleGrabberCallback);
+			OutDimensions = SetPlaybackTopology(SampleGrabberCallback, bUseSound);
 		}
 		else
 		{
@@ -490,7 +542,7 @@ void FVideoPlayer::Shutdown()
 	}
 }
 
-FIntPoint FVideoPlayer::SetPlaybackTopology(FSampleGrabberCallback* SampleGrabberCallback)
+FIntPoint FVideoPlayer::SetPlaybackTopology(FSampleGrabberCallback* SampleGrabberCallback, bool bUseSound)
 {
 	FIntPoint OutDimensions = FIntPoint(ForceInit);
 
@@ -517,7 +569,7 @@ FIntPoint FVideoPlayer::SetPlaybackTopology(FSampleGrabberCallback* SampleGrabbe
 
 		if (bSelected)
 		{
-			FIntPoint VideoDimensions = AddStreamToTopology(Topology, PresentationDesc, StreamDesc, SampleGrabberCallback);
+			FIntPoint VideoDimensions = AddStreamToTopology(Topology, PresentationDesc, StreamDesc, SampleGrabberCallback, bUseSound);
 			if (VideoDimensions != FIntPoint(ForceInit))
 			{
 				OutDimensions = VideoDimensions;
@@ -536,7 +588,7 @@ FIntPoint FVideoPlayer::SetPlaybackTopology(FSampleGrabberCallback* SampleGrabbe
 	return OutDimensions;
 }
 
-FIntPoint FVideoPlayer::AddStreamToTopology(IMFTopology* Topology, IMFPresentationDescriptor* PresentationDesc, IMFStreamDescriptor* StreamDesc, FSampleGrabberCallback* SampleGrabberCallback)
+FIntPoint FVideoPlayer::AddStreamToTopology(IMFTopology* Topology, IMFPresentationDescriptor* PresentationDesc, IMFStreamDescriptor* StreamDesc, FSampleGrabberCallback* SampleGrabberCallback, bool bUseSound)
 {
 	FIntPoint OutDimensions = FIntPoint(ForceInit);
 
@@ -554,6 +606,12 @@ FIntPoint FVideoPlayer::AddStreamToTopology(IMFTopology* Topology, IMFPresentati
 
 		if (MajorType == MFMediaType_Audio)
 		{
+			if (!bUseSound)
+			{
+				SAFE_RELEASE(Handler);
+				return FIntPoint(ForceInit);
+			}
+
 			HResult = MFCreateAudioRendererActivate(&SinkActivate);
 			check(SUCCEEDED(HResult));
 
diff --git a/Engine/Plugins/Runtime/WindowsMoviePlayer/Source/WindowsMoviePlayer/Private/WindowsMovieStreamer.h b/Engine/Plugins/Runtime/WindowsMoviePlayer/Source/WindowsMoviePlayer/Private/WindowsMovieStreamer.h
--- a/Engine/Plugins/Runtime/WindowsMoviePlayer/Source/WindowsMoviePlayer/Private/WindowsMovieStreamer.h
+++ b/Engine/Plugins/Runtime/WindowsMoviePlayer/Source/WindowsMoviePlayer/Private/WindowsMovieStreamer.h
@@ -83,6 +83,9 @@ private:
 	/** The video player and sample grabber for use of Media Foundation */
 	class FVideoPlayer* VideoPlayer;
     class FSampleGrabberCallback* SampleGrabberCallback;
+
+	/** Defaults to true, modified to false if video fails to play allowing a retry without sound */
+	bool bUseSound;
 };
 
 /** Video track details */
@@ -112,7 +115,11 @@ public:
 		, MediaSession(NULL)
 		, MediaSource(NULL)
 		, MovieIsFinished(0)
-		, CloseIsPosted(0) {}
+		, CloseIsPosted(0)
+		, bFailedToCreateMediaSink(false)
+	{
+	}
+
 	virtual ~FVideoPlayer()
 	{
 		check(MediaSession == NULL);
@@ -129,7 +136,7 @@ public:
 	STDMETHODIMP Invoke(IMFAsyncResult* AsyncResult);
 
 	/** Opens the specified file, and returns the video dimensions */
-	FIntPoint OpenFile(const FString& FilePath, class FSampleGrabberCallback* SampleGrabberCallback);
+	FIntPoint OpenFile(const FString& FilePath, class FSampleGrabberCallback* SampleGrabberCallback, bool bUseSound);
 	/** Starts the video player playback */
 	void StartPlayback();
 	/** Shuts the video player down, destroying all threads with it */
@@ -140,11 +147,14 @@ public:
 	
 	const FMovieTrackFormat& GetVideoTrackFormat() const { return VideoTrackFormat; }
 
+	/** True when there was an error related to creating a media sink (MESessionTopologySet event received with status MF_E_CANNOT_CREATE_SINK) */
+	bool FailedToCreateMediaSink() const { return bFailedToCreateMediaSink; }
+
 private:
 	/** Sets up the topology of all the nodes in the media session, returning the video dimensions */
-	FIntPoint SetPlaybackTopology(class FSampleGrabberCallback* SampleGrabberCallback);
+	FIntPoint SetPlaybackTopology(class FSampleGrabberCallback* SampleGrabberCallback, bool bUseSound);
 	/** Adds a single audio or video stream to the passed in topology, returning video dimensions if applicable */
-	FIntPoint AddStreamToTopology(IMFTopology* Topology, IMFPresentationDescriptor* PresentationDesc, IMFStreamDescriptor* StreamDesc, class FSampleGrabberCallback* SampleGrabberCallback);
+	FIntPoint AddStreamToTopology(IMFTopology* Topology, IMFPresentationDescriptor* PresentationDesc, IMFStreamDescriptor* StreamDesc, class FSampleGrabberCallback* SampleGrabberCallback, bool bUseSound);
 
 private:
 	/** Media Foundation boilerplate */
@@ -162,6 +172,9 @@ private:
 
 	/* Format of video track */
 	FMovieTrackFormat VideoTrackFormat;
+
+	/** True when there was an error related to creating a media sink (MESessionTopologySet event received with status MF_E_CANNOT_CREATE_SINK) */
+	bool bFailedToCreateMediaSink;
 };
 
 
