Updated Windows Movie Streamer (startup movies) to use render pass

diff --git a/Engine/Plugins/Runtime/WindowsMoviePlayer/Source/WindowsMoviePlayer/Private/WindowsMovieStreamer.cpp b/Engine/Plugins/Runtime/WindowsMoviePlayer/Source/WindowsMoviePlayer/Private/WindowsMovieStreamer.cpp
--- a/Engine/Plugins/Runtime/WindowsMoviePlayer/Source/WindowsMoviePlayer/Private/WindowsMovieStreamer.cpp
+++ b/Engine/Plugins/Runtime/WindowsMoviePlayer/Source/WindowsMoviePlayer/Private/WindowsMovieStreamer.cpp
@@ -117,55 +117,59 @@ void FMediaFoundationMovieStreamer::ConvertSample()
 
 	// perform the conversion
 	FRHICommandListImmediate& CommandList = FRHICommandListExecutor::GetImmediateCommandList();
-	
-	FGraphicsPipelineStateInitializer GraphicsPSOInit;
-	SetRenderTargets(CommandList, 1, &RenderTarget, nullptr, ESimpleRenderTargetMode::EExistingColorAndDepth, FExclusiveDepthStencil::DepthNop_StencilNop);
 
-	CommandList.ApplyCachedRenderTargets(GraphicsPSOInit);
-	CommandList.SetViewport(0, 0, 0.0f, OutputDim.X, OutputDim.Y, 1.0f);
+	FRHIRenderPassInfo RPInfo(RenderTarget, ERenderTargetActions::Load_Store);
+	CommandList.BeginRenderPass(RPInfo, TEXT("WindowsMovieConvertSample"));
+	{
+		FGraphicsPipelineStateInitializer GraphicsPSOInit;
 
-	GraphicsPSOInit.DepthStencilState = TStaticDepthStencilState<false, CF_Always>::GetRHI();
-	GraphicsPSOInit.RasterizerState = TStaticRasterizerState<>::GetRHI();
-	GraphicsPSOInit.BlendState = TStaticBlendStateWriteMask<CW_RGBA, CW_NONE, CW_NONE, CW_NONE, CW_NONE, CW_NONE, CW_NONE, CW_NONE>::GetRHI();
-	GraphicsPSOInit.PrimitiveType = PT_TriangleStrip;
+		CommandList.ApplyCachedRenderTargets(GraphicsPSOInit);
+		CommandList.SetViewport(0, 0, 0.0f, OutputDim.X, OutputDim.Y, 1.0f);
 
-	// configure media shaders
-	auto ShaderMap = GetGlobalShaderMap(GMaxRHIFeatureLevel);
-	TShaderMapRef<FMediaShadersVS> VertexShader(ShaderMap);
+		GraphicsPSOInit.DepthStencilState = TStaticDepthStencilState<false, CF_Always>::GetRHI();
+		GraphicsPSOInit.RasterizerState = TStaticRasterizerState<>::GetRHI();
+		GraphicsPSOInit.BlendState = TStaticBlendStateWriteMask<CW_RGBA, CW_NONE, CW_NONE, CW_NONE, CW_NONE, CW_NONE, CW_NONE, CW_NONE>::GetRHI();
+		GraphicsPSOInit.PrimitiveType = PT_TriangleStrip;
 
-	GraphicsPSOInit.BoundShaderState.VertexDeclarationRHI = GMediaVertexDeclaration.VertexDeclarationRHI;
-	GraphicsPSOInit.BoundShaderState.VertexShaderRHI = GETSAFERHISHADER_VERTEX(*VertexShader);
+		// configure media shaders
+		auto ShaderMap = GetGlobalShaderMap(GMaxRHIFeatureLevel);
+		TShaderMapRef<FMediaShadersVS> VertexShader(ShaderMap);
 
-	switch (SourceFormat.SampleFormat)
-	{
-	case EMediaTextureSampleFormat::CharBMP:
-	{
-		TShaderMapRef<FBMPConvertPS> ConvertShader(ShaderMap);
-		GraphicsPSOInit.BoundShaderState.PixelShaderRHI = GETSAFERHISHADER_PIXEL(*ConvertShader);
-		SetGraphicsPipelineState(CommandList, GraphicsPSOInit);
-		ConvertShader->SetParameters(CommandList, InputTarget, OutputDim, bSampleIsOutputSrgb && !SrgbOutput);
-	}
-	break;
-	
-	case EMediaTextureSampleFormat::CharYUY2:
-	{
-		TShaderMapRef<FYUY2ConvertPS> ConvertShader(ShaderMap);
-		GraphicsPSOInit.BoundShaderState.PixelShaderRHI = GETSAFERHISHADER_PIXEL(*ConvertShader);
-		SetGraphicsPipelineState(CommandList, GraphicsPSOInit);
-		ConvertShader->SetParameters(CommandList, InputTarget, OutputDim, MediaShaders::YuvToSrgbDefault, bSampleIsOutputSrgb);
-	}
-	break;
+		GraphicsPSOInit.BoundShaderState.VertexDeclarationRHI = GMediaVertexDeclaration.VertexDeclarationRHI;
+		GraphicsPSOInit.BoundShaderState.VertexShaderRHI = GETSAFERHISHADER_VERTEX(*VertexShader);
 
-	default:
-		return; // unsupported format
-	}
+		switch (SourceFormat.SampleFormat)
+		{
+		case EMediaTextureSampleFormat::CharBMP:
+		{
+			TShaderMapRef<FBMPConvertPS> ConvertShader(ShaderMap);
+			GraphicsPSOInit.BoundShaderState.PixelShaderRHI = GETSAFERHISHADER_PIXEL(*ConvertShader);
+			SetGraphicsPipelineState(CommandList, GraphicsPSOInit);
+			ConvertShader->SetParameters(CommandList, InputTarget, OutputDim, bSampleIsOutputSrgb && !SrgbOutput);
+		}
+		break;
+
+		case EMediaTextureSampleFormat::CharYUY2:
+		{
+			TShaderMapRef<FYUY2ConvertPS> ConvertShader(ShaderMap);
+			GraphicsPSOInit.BoundShaderState.PixelShaderRHI = GETSAFERHISHADER_PIXEL(*ConvertShader);
+			SetGraphicsPipelineState(CommandList, GraphicsPSOInit);
+			ConvertShader->SetParameters(CommandList, InputTarget, OutputDim, MediaShaders::YuvToSrgbDefault, bSampleIsOutputSrgb);
+		}
+		break;
 
-	// draw full size quad into render target
-	FVertexBufferRHIRef VertexBuffer = CreateTempMediaVertexBuffer();
-	CommandList.SetStreamSource(0, VertexBuffer, 0);
-	// set viewport to RT size
-	CommandList.SetViewport(0, 0, 0.0f, OutputDim.X, OutputDim.Y, 1.0f);
-	CommandList.DrawPrimitive(PT_TriangleStrip, 0, 2, 1);
+		default:
+			return; // unsupported format
+		}
+
+		// draw full size quad into render target
+		FVertexBufferRHIRef VertexBuffer = CreateTempMediaVertexBuffer();
+		CommandList.SetStreamSource(0, VertexBuffer, 0);
+		// set viewport to RT size
+		CommandList.SetViewport(0, 0, 0.0f, OutputDim.X, OutputDim.Y, 1.0f);
+		CommandList.DrawPrimitive(0, 2, 1);
+	}
+	CommandList.EndRenderPass();
 	CommandList.TransitionResource(EResourceTransitionAccess::EReadable, RenderTarget);
 }
 
