Windows startup movies now can handle more sitatuions that cause playback to be interrupted.

diff --git a/Engine/Plugins/Runtime/WindowsMoviePlayer/Source/WindowsMoviePlayer/Private/WindowsMovieStreamer.cpp b/Engine/Plugins/Runtime/WindowsMoviePlayer/Source/WindowsMoviePlayer/Private/WindowsMovieStreamer.cpp
--- a/Engine/Plugins/Runtime/WindowsMoviePlayer/Source/WindowsMoviePlayer/Private/WindowsMovieStreamer.cpp
+++ b/Engine/Plugins/Runtime/WindowsMoviePlayer/Source/WindowsMoviePlayer/Private/WindowsMovieStreamer.cpp
@@ -400,7 +400,7 @@ HRESULT FVideoPlayer::Invoke(IMFAsyncResult* AsyncResult)
 	bool bFinishedAndClose = false;
 	switch (EventType)
 	{
-	case MEAudioSessionDeviceRemoved:
+	case MESinkInvalidated:
 		// Need to stop playback now or will be stuck forever
 		bFinishedAndClose = true;
 		break;
