Merged: Startup movies on Windows platform can now be 1920x1080.
Previously the resolution needed to be a multiple of 16 (such as 1920x1088) due to using MFVideoFormat_RGB32 setting.

All windows startup movies are now processed as YUY2 and drawn using conversion shader.

diff --git a/Engine/Plugins/Runtime/WindowsMoviePlayer/Source/WindowsMoviePlayer/WindowsMoviePlayer.Build.cs b/Engine/Plugins/Runtime/WindowsMoviePlayer/Source/WindowsMoviePlayer/WindowsMoviePlayer.Build.cs
--- a/Engine/Plugins/Runtime/WindowsMoviePlayer/Source/WindowsMoviePlayer/WindowsMoviePlayer.Build.cs
+++ b/Engine/Plugins/Runtime/WindowsMoviePlayer/Source/WindowsMoviePlayer/WindowsMoviePlayer.Build.cs
@@ -33,6 +33,7 @@ namespace UnrealBuildTool.Rules
                     "RHI",
                     "SlateCore",
                     "Slate",
+					"ShaderCore",
 					"UtilityShaders"
 				}
 				);
