Merge fix

diff --git a/Engine/Plugins/Runtime/WindowsMoviePlayer/Source/WindowsMoviePlayer/WindowsMoviePlayer.Build.cs b/Engine/Plugins/Runtime/WindowsMoviePlayer/Source/WindowsMoviePlayer/WindowsMoviePlayer.Build.cs
--- a/Engine/Plugins/Runtime/WindowsMoviePlayer/Source/WindowsMoviePlayer/WindowsMoviePlayer.Build.cs
+++ b/Engine/Plugins/Runtime/WindowsMoviePlayer/Source/WindowsMoviePlayer/WindowsMoviePlayer.Build.cs
@@ -33,7 +33,6 @@ namespace UnrealBuildTool.Rules
                     "RHI",
                     "SlateCore",
                     "Slate",
-					"ShaderCore",
 					"UtilityShaders"
 				}
 				);
