WindowsMoviePlayer: Fixed crash when playing 2 startup movies.

diff --git a/Engine/Plugins/Runtime/WindowsMoviePlayer/Source/WindowsMoviePlayer/Private/WindowsMovieStreamer.cpp b/Engine/Plugins/Runtime/WindowsMoviePlayer/Source/WindowsMoviePlayer/Private/WindowsMovieStreamer.cpp
--- a/Engine/Plugins/Runtime/WindowsMoviePlayer/Source/WindowsMoviePlayer/Private/WindowsMovieStreamer.cpp
+++ b/Engine/Plugins/Runtime/WindowsMoviePlayer/Source/WindowsMoviePlayer/Private/WindowsMovieStreamer.cpp
@@ -616,15 +616,22 @@ FIntPoint FVideoPlayer::AddStreamToTopology(IMFTopology* Topology, IMFPresentati
 			check(SUCCEEDED(HResult));
 
 			// Allow HMD, if present, to override audio output device
-			if (IHeadMountedDisplayModule::IsAvailable())
+			// Can only access HMD from the game thread.
+			if (IsInGameThread())
 			{
-				FString AudioOutputDevice = IHeadMountedDisplayModule::Get().GetAudioOutputDevice();
-
-				if(!AudioOutputDevice.IsEmpty())
+				if (IHeadMountedDisplayModule::IsAvailable())
 				{
-					HResult = SinkActivate->SetString(MF_AUDIO_RENDERER_ATTRIBUTE_ENDPOINT_ID, *AudioOutputDevice);
-					check(SUCCEEDED(HResult));
+					HMDAudioOutputDevice = IHeadMountedDisplayModule::Get().GetAudioOutputDevice();
 				}
+				else
+				{
+					HMDAudioOutputDevice.Reset();
+				}
+			}
+			if(!HMDAudioOutputDevice.IsEmpty())
+			{
+				HResult = SinkActivate->SetString(MF_AUDIO_RENDERER_ATTRIBUTE_ENDPOINT_ID, *HMDAudioOutputDevice);
+				check(SUCCEEDED(HResult));
 			}
 		}
 		else if (MajorType == MFMediaType_Video)
diff --git a/Engine/Plugins/Runtime/WindowsMoviePlayer/Source/WindowsMoviePlayer/Private/WindowsMovieStreamer.h b/Engine/Plugins/Runtime/WindowsMoviePlayer/Source/WindowsMoviePlayer/Private/WindowsMovieStreamer.h
--- a/Engine/Plugins/Runtime/WindowsMoviePlayer/Source/WindowsMoviePlayer/Private/WindowsMovieStreamer.h
+++ b/Engine/Plugins/Runtime/WindowsMoviePlayer/Source/WindowsMoviePlayer/Private/WindowsMovieStreamer.h
@@ -173,6 +173,9 @@ private:
 	/* Format of video track */
 	FMovieTrackFormat VideoTrackFormat;
 
+	/* Caches the audio device from the HMD. */
+	FString HMDAudioOutputDevice;
+
 	/** True when there was an error related to creating a media sink (MESessionTopologySet event received with status MF_E_CANNOT_CREATE_SINK) */
 	bool bFailedToCreateMediaSink;
 };
