WindowsMovieStreamer.cpp static analysis fix.  Compile time zero mask test is always false due to "const bool SrgbTexture = false".

#rnx

diff --git a/Engine/Plugins/Runtime/WindowsMoviePlayer/Source/WindowsMoviePlayer/Private/WindowsMovieStreamer.cpp b/Engine/Plugins/Runtime/WindowsMoviePlayer/Source/WindowsMoviePlayer/Private/WindowsMovieStreamer.cpp
--- a/Engine/Plugins/Runtime/WindowsMoviePlayer/Source/WindowsMoviePlayer/Private/WindowsMovieStreamer.cpp
+++ b/Engine/Plugins/Runtime/WindowsMoviePlayer/Source/WindowsMoviePlayer/Private/WindowsMovieStreamer.cpp
@@ -86,7 +86,7 @@ void FMediaFoundationMovieStreamer::ConvertSample(FRHICommandListImmediate& RHIC
 		const ETextureCreateFlags InputCreateFlags = SrgbTexture ? TexCreate_SRGB : TexCreate_None;
 
 		// create a new input render target if necessary
-		if (!InputTarget.IsValid() || (InputTarget->GetSizeXY() != SourceFormat.BufferDim) || (InputTarget->GetFormat() != InputPixelFormat) || ((InputTarget->GetFlags() & InputCreateFlags) != InputCreateFlags))
+		if (!InputTarget.IsValid() || (InputTarget->GetSizeXY() != SourceFormat.BufferDim) || (InputTarget->GetFormat() != InputPixelFormat) || ((InputTarget->GetFlags() & TexCreate_SRGB) != InputCreateFlags))
 		{
 			const FRHITextureCreateDesc Desc =
 				FRHITextureCreateDesc::Create2D(TEXT("FMediaFoundationMovieStreamer"))
