Cleaning up a few movie streamer issues.

diff --git a/Engine/Plugins/Runtime/WindowsMoviePlayer/Source/WindowsMoviePlayer/Private/WindowsMovieStreamer.cpp b/Engine/Plugins/Runtime/WindowsMoviePlayer/Source/WindowsMoviePlayer/Private/WindowsMovieStreamer.cpp
--- a/Engine/Plugins/Runtime/WindowsMoviePlayer/Source/WindowsMoviePlayer/Private/WindowsMovieStreamer.cpp
+++ b/Engine/Plugins/Runtime/WindowsMoviePlayer/Source/WindowsMoviePlayer/Private/WindowsMovieStreamer.cpp
@@ -131,6 +131,7 @@ void FMediaFoundationMovieStreamer::ConvertSample(FRHICommandList& RHICmdList)
 		GraphicsPSOInit.BoundShaderState.VertexDeclarationRHI = GMediaVertexDeclaration.VertexDeclarationRHI;
 		GraphicsPSOInit.BoundShaderState.VertexShaderRHI = VertexShader.GetVertexShader();
 
+		bool bValidFormat = true;
 		switch (SourceFormat.SampleFormat)
 		{
 		case EMediaTextureSampleFormat::CharBMP:
@@ -152,15 +153,18 @@ void FMediaFoundationMovieStreamer::ConvertSample(FRHICommandList& RHICmdList)
 		break;
 
 		default:
-			return; // unsupported format
+			bValidFormat = false; // unsupported format
 		}
 
-		// draw full size quad into render target
-		FBufferRHIRef VertexBuffer = CreateTempMediaVertexBuffer(RHICmdList);
-		RHICmdList.SetStreamSource(0, VertexBuffer, 0);
-		// set viewport to RT size
-		RHICmdList.SetViewport(0, 0, 0.0f, OutputDim.X, OutputDim.Y, 1.0f);
-		RHICmdList.DrawPrimitive(0, 2, 1);
+		if (bValidFormat)
+		{
+			// draw full size quad into render target
+			FBufferRHIRef VertexBuffer = CreateTempMediaVertexBuffer(RHICmdList);
+			RHICmdList.SetStreamSource(0, VertexBuffer, 0);
+			// set viewport to RT size
+			RHICmdList.SetViewport(0, 0, 0.0f, OutputDim.X, OutputDim.Y, 1.0f);
+			RHICmdList.DrawPrimitive(0, 2, 1);
+		}
 	}
 	RHICmdList.EndRenderPass();
 	RHICmdList.Transition(FRHITransitionInfo(RenderTarget, ERHIAccess::Unknown, ERHIAccess::SRVGraphics));
